Fix for non-unity compile issue in WindowsMixedReaility

diff --git a/Engine/Source/Runtime/WindowsMixedReality/WindowsMixedRealityRuntimeSettings/Private/WindowsMixedRealityRuntimeSettings.cpp b/Engine/Source/Runtime/WindowsMixedReality/WindowsMixedRealityRuntimeSettings/Private/WindowsMixedRealityRuntimeSettings.cpp
--- a/Engine/Source/Runtime/WindowsMixedReality/WindowsMixedRealityRuntimeSettings/Private/WindowsMixedRealityRuntimeSettings.cpp
+++ b/Engine/Source/Runtime/WindowsMixedReality/WindowsMixedRealityRuntimeSettings/Private/WindowsMixedRealityRuntimeSettings.cpp
@@ -3,6 +3,7 @@
 #include "WindowsMixedRealityRuntimeSettings.h"
 #include "Misc/ConfigCacheIni.h"
 #include "CoreGlobals.h"
+#include "Package.h"
 
 UWindowsMixedRealityRuntimeSettings* UWindowsMixedRealityRuntimeSettings::WMRSettingsSingleton = nullptr;
 
