Fixing full path for include

diff --git a/Engine/Source/Runtime/WindowsMixedReality/WindowsMixedRealityRuntimeSettings/Private/WindowsMixedRealityRuntimeSettings.cpp b/Engine/Source/Runtime/WindowsMixedReality/WindowsMixedRealityRuntimeSettings/Private/WindowsMixedRealityRuntimeSettings.cpp
--- a/Engine/Source/Runtime/WindowsMixedReality/WindowsMixedRealityRuntimeSettings/Private/WindowsMixedRealityRuntimeSettings.cpp
+++ b/Engine/Source/Runtime/WindowsMixedReality/WindowsMixedRealityRuntimeSettings/Private/WindowsMixedRealityRuntimeSettings.cpp
@@ -3,7 +3,7 @@
 #include "WindowsMixedRealityRuntimeSettings.h"
 #include "Misc/ConfigCacheIni.h"
 #include "CoreGlobals.h"
-#include "Package.h"
+#include "UObject/Package.h"
 
 UWindowsMixedRealityRuntimeSettings* UWindowsMixedRealityRuntimeSettings::WMRSettingsSingleton = nullptr;
 
